First pass on a skeleton for the Skein CLI
* boilerplate code for all the commands that are currently listed in the spec
* support for required and optional arguments and options
* basic user interaction prompts
* basic support for progress bars.

diff --git a/Engine/Programs/SkeinCLI/SkeinCLI/AssetsCmd.cs b/Engine/Programs/SkeinCLI/SkeinCLI/AssetsCmd.cs
new file mode 100644
--- /dev/null
+++ b/Engine/Programs/SkeinCLI/SkeinCLI/AssetsCmd.cs
@@ -0,0 +1,217 @@
+﻿using McMaster.Extensions.CommandLineUtils;
+using Serilog;
+using System;
+using System.Collections.Generic;
+using System.ComponentModel.DataAnnotations;
+using System.IO;
+using System.Threading.Tasks;
+
+namespace SkeinCLI
+{
+
+    [Command(Name = "assets", Description = "Executes assets related commands")]
+    [Subcommand(
+        typeof(AssetsTrackCmd),
+        typeof(AssetsUntrackCmd),
+        typeof(AssetsGetCmd),
+        typeof(AssetsRenameCmd),
+        typeof(AssetsRevertCmd))]
+    class AssetsCmd : SkeinCmdBase
+    {
+        // This will automatically be set before OnExecute is invoked.
+        private SkeinCmd Parent { get; set; }
+
+        public override List<string> CreateArgs()
+        {
+            var args = Parent.CreateArgs();
+            args.Add("assets");
+
+            return args;
+        }
+
+        protected override Task<int> OnExecute(CommandLineApplication app)
+        {
+            //'skein assets' is an incomplete command, so we show help here even if --help is not specified
+            app.ShowHelp();
+            return base.OnExecute(app);
+        }
+    }
+
+    #region assets track
+
+    [Command(Name = "track", Description = "Track an asset for the project")]
+    class AssetsTrackCmd : SkeinCmdBase
+    {
+        private AssetsCmd Parent { get; set; }
+
+        [Argument(0, Name = "filepath", Description = "File path for the asset to track")]
+        [FileExists]
+        [Required]
+        public string FilePath { get; set; }
+
+        public override List<string> CreateArgs()
+        {
+            var args = Parent.CreateArgs();
+            args.Add("track");
+
+            return args;
+        }
+
+        protected override Task<int> OnExecute(CommandLineApplication app)
+        {
+            Log.Logger.ForContext<ProjectsSnapshotsGetCmd>().Information($"File Path: {FilePath}");
+            Log.Logger.ForContext<AssetsTrackCmd>().Information("Result of track command here [...]");
+            return base.OnExecute(app);
+        }
+    }
+
+    #endregion
+
+    #region assets untrack
+
+    [Command(Name = "untrack", Description = "Untrack an asset for the project")]
+    class AssetsUntrackCmd : SkeinCmdBase
+    {
+        private AssetsCmd Parent { get; set; }
+
+        [Argument(0, Name = "filepath", Description = "File path for the asset to untrack")]
+        [FileExists]
+        [Required]
+        public string FilePath { get; set; }
+
+        public override List<string> CreateArgs()
+        {
+            var args = Parent.CreateArgs();
+            args.Add("untrack");
+
+            return args;
+        }
+
+        protected override Task<int> OnExecute(CommandLineApplication app)
+        {
+            Log.Logger.ForContext<ProjectsSnapshotsGetCmd>().Information($"File Path: {FilePath}");
+            Log.Logger.ForContext<AssetsUntrackCmd>().Information("Result of untrack command here [...]");
+            return base.OnExecute(app);
+        }
+    }
+
+    #endregion
+
+    #region assets get
+
+    [Command(Name = "get", Description = "Get an asset, update skein.yml")]
+    class AssetsGetCmd : SkeinCmdBase
+    {
+        private AssetsCmd Parent { get; set; }
+
+        [Argument(0, Name = "asset-id", Description = "uuid of the asset to get")]
+        [Required]
+        public string AssetId { get; set; }
+
+        [Argument(1, Name = "path", Description = "Local path where to download the asset, defaults to the current path")]
+        [DirectoryExists]
+        public string Path { get; set; }
+
+        public override List<string> CreateArgs()
+        {
+            var args = Parent.CreateArgs();
+            args.Add("get");
+
+            return args;
+        }
+
+        protected override Task<int> OnExecute(CommandLineApplication app)
+        {
+            if (string.IsNullOrEmpty(Path))
+            {
+                Path = Directory.GetCurrentDirectory();
+            }
+            Log.Logger.ForContext<ProjectsSnapshotsGetCmd>().Information($"AssetId: {AssetId}, Path: {Path}");
+            Log.Logger.ForContext<AssetsGetCmd>().Information("Result of assets get described here [...]");
+            return base.OnExecute(app);
+        }
+    }
+
+    #endregion
+
+    #region assets rename
+
+    [Command(Name = "rename", Description = "Rename an asset and update skein.yml")]
+    class AssetsRenameCmd : SkeinCmdBase
+    {
+        private AssetsCmd Parent { get; set; }
+
+        [Argument(0, Name = "old-filepath", Description = "File path for the asset to rename")]
+        [FileExists]
+        [Required]
+        public string OldFilePath { get; set; }
+
+        [Argument(1, Name = "new-filepath", Description = "New path for the renamed asset")]
+        [FileNotExists]
+        [Required]
+        public string NewFilePath { get; set; }
+
+        public override List<string> CreateArgs()
+        {
+            var args = Parent.CreateArgs();
+            args.Add("rename");
+
+            return args;
+        }
+
+        protected override Task<int> OnExecute(CommandLineApplication app)
+        {
+            Log.Logger.ForContext<ProjectsSnapshotsGetCmd>().Information($"Old File Path: {OldFilePath}, New File Path: {NewFilePath}");
+            Log.Logger.ForContext<AssetsRenameCmd>().Information("Result of assets rename described here [...]");
+            return base.OnExecute(app);
+        }
+    }
+
+    #endregion
+
+    #region assets revert
+
+    [Command(Name = "revert", Description = "Revert the project to the snapshot")]
+    class AssetsRevertCmd : SkeinCmdBase
+    {
+        private AssetsCmd Parent { get; set; }
+
+        [Argument(0, Name = "filepath", Description = "File path for the asset to revert")]
+        [FileExists]
+        [Required]
+        public string FilePath { get; set; }
+
+        public override List<string> CreateArgs()
+        {
+            var args = Parent.CreateArgs();
+            args.Add("revert");
+
+            return args;
+        }
+
+        protected override Task<int> OnExecute(CommandLineApplication app)
+        {
+            bool proceed = true;
+            if (!Quiet)
+            {
+                proceed = Prompt.GetYesNo($"Are you sure you want to revert the asset at {FilePath}?",
+                    defaultAnswer: true,
+                    promptColor: ConsoleColor.Black,
+                    promptBgColor: ConsoleColor.White);
+            }
+
+            if (proceed)
+            {
+                Log.Logger.ForContext<ProjectsSnapshotsGetCmd>().Information($"File Path: {FilePath}");
+                Log.Logger.ForContext<AssetsRevertCmd>().Information("Assets Revert operation result here [...]");
+            }
+            else
+            {
+                Log.Logger.ForContext<ProjectsRevertCmd>().Information("User cancelled operation");
+            }
+            return base.OnExecute(app);
+        }
+    }
+
+    #endregion
+}
\ No newline at end of file
diff --git a/Engine/Programs/SkeinCLI/SkeinCLI/AuthCmd.cs b/Engine/Programs/SkeinCLI/SkeinCLI/AuthCmd.cs
new file mode 100644
--- /dev/null
+++ b/Engine/Programs/SkeinCLI/SkeinCLI/AuthCmd.cs
@@ -0,0 +1,101 @@
+﻿using McMaster.Extensions.CommandLineUtils;
+using Serilog;
+using System;
+using System.Collections.Generic;
+using System.Threading.Tasks;
+
+namespace SkeinCLI
+{
+    [Command(Name = "auth", Description = "Executes authorization related commands")]
+    [Subcommand(typeof(AuthLoginCmd),
+        typeof(AuthLogoutCmd))]
+    class AuthCmd : SkeinCmdBase
+    {
+        private SkeinCmdBase Parent { get; set; }
+
+        public override List<string> CreateArgs()
+        {
+            var args = Parent.CreateArgs();
+            args.Add("auth");
+
+            return args;
+        }
+
+        protected override Task<int> OnExecute(CommandLineApplication app)
+        {
+            //'skein auth' is an incomplete command, so we show help here even if --help is not specified
+            app.ShowHelp();
+            return base.OnExecute(app);
+        }
+    }
+
+    [Command(Name = "login", Description = "Skein login")]
+    class AuthLoginCmd : SkeinCmdBase
+    {
+        private AuthCmd Parent { get; set; }
+
+        [Option("-u|--username", CommandOptionType.SingleValue, Description = "User Name")]
+        public string UserName { get; set; }
+
+        [Option("-p|--password", CommandOptionType.SingleValue, Description = "Password")]
+        public string Password { get; set; }
+
+        public override List<string> CreateArgs()
+        {
+            var args = Parent.CreateArgs();
+            args.Add("login");
+
+            return args;
+        }
+
+        protected override Task<int> OnExecute(CommandLineApplication app)
+        {
+            if (string.IsNullOrEmpty(UserName) && !Quiet)
+            {
+                UserName = Prompt.GetString("User Name:",
+                    promptColor: ConsoleColor.Black,
+                    promptBgColor: ConsoleColor.White);
+            }
+
+            if (string.IsNullOrEmpty(Password) && !Quiet)
+            {
+                Password = Prompt.GetPassword("Password:",
+                    promptColor: ConsoleColor.Black,
+                    promptBgColor: ConsoleColor.White);
+            }
+
+            if (!string.IsNullOrEmpty(UserName) && !string.IsNullOrEmpty(Password))
+            {
+                Log.Logger.ForContext<ProjectsSnapshotsGetCmd>().Information($"User Name: {UserName}, Password: {Password}");
+                Log.Logger.ForContext<ProjectsListCmd>().Information("Login process here [...]");
+            }
+            else 
+            {
+                Log.Logger.ForContext<ProjectsListCmd>().Information("Login cancelled");
+            }
+
+            
+            return base.OnExecute(app);
+        }
+    }
+
+    [Command(Name = "logout", Description = "Skein logout")]
+    class AuthLogoutCmd : SkeinCmdBase
+    {
+        private AuthCmd Parent { get; set; }
+
+        public override List<string> CreateArgs()
+        {
+            var args = Parent.CreateArgs();
+            args.Add("logout");
+
+            return args;
+        }
+
+        protected override Task<int> OnExecute(CommandLineApplication app)
+        {
+            Log.Logger.ForContext<ProjectsListCmd>().Information("Logout process here [...]");
+            return base.OnExecute(app);
+        }
+    }
+}
\ No newline at end of file
diff --git a/Engine/Programs/SkeinCLI/SkeinCLI/Program.cs b/Engine/Programs/SkeinCLI/SkeinCLI/Program.cs
new file mode 100644
--- /dev/null
+++ b/Engine/Programs/SkeinCLI/SkeinCLI/Program.cs
@@ -0,0 +1,49 @@
+﻿using McMaster.Extensions.CommandLineUtils;
+using Microsoft.Extensions.Configuration;
+using Microsoft.Extensions.Configuration.Json;
+using Serilog;
+using System;
+using System.IO;
+using System.Threading.Tasks;
+
+namespace SkeinCLI
+{
+    class Program
+    {
+        private static async Task<int> Main(string[] args)
+        {
+            var Configuration = new ConfigurationBuilder()
+                .SetBasePath(Directory.GetCurrentDirectory())
+                .AddJsonFile(AppDomain.CurrentDomain.BaseDirectory + "\\appsettings.json", optional: true, reloadOnChange: true)
+                .AddEnvironmentVariables()
+                .Build();
+
+            Log.Logger = new LoggerConfiguration()
+               .ReadFrom.Configuration(Configuration)
+               .CreateLogger();
+
+            /*using (var progress = new ProgressBar())
+            {
+                for (int i = 0; i <= 100; i++)
+                {
+                    progress.Report((double)i / 100);
+                    System.Threading.Thread.Sleep(20);
+                }
+            }*/
+
+            try
+            {
+                return await CommandLineApplication.ExecuteAsync<SkeinCmd>(args);
+            }
+            catch (Exception ex)
+            {
+                Log.ForContext<Program>().Error(ex, "Oh no!");
+                return 1;
+            }
+            finally 
+            {
+                Log.CloseAndFlush();
+            }
+        }
+    }
+}
diff --git a/Engine/Programs/SkeinCLI/SkeinCLI/ProgressBar.cs b/Engine/Programs/SkeinCLI/SkeinCLI/ProgressBar.cs
new file mode 100644
--- /dev/null
+++ b/Engine/Programs/SkeinCLI/SkeinCLI/ProgressBar.cs
@@ -0,0 +1,107 @@
+﻿/*
+ * From https://gist.github.com/DanielSWolf/0ab6a96899cc5377bf54
+ * Licensed under the MIT License: http://opensource.org/licenses/MIT
+ */
+
+using System;
+using System.Text;
+using System.Threading;
+
+/// <summary>
+/// An ASCII progress bar
+/// </summary>
+public class ProgressBar : IDisposable, IProgress<double>
+{
+	private const int blockCount = 40;
+	private readonly TimeSpan animationInterval = TimeSpan.FromSeconds(1.0 / 8);
+	private const string animation = @"|/-\";
+
+	private readonly Timer timer;
+
+	private double currentProgress = 0;
+	private string currentText = string.Empty;
+	private bool disposed = false;
+	private int animationIndex = 0;
+
+	public ProgressBar()
+	{
+		timer = new Timer(TimerHandler);
+
+		// A progress bar is only for temporary display in a console window.
+		// If the console output is redirected to a file, draw nothing.
+		// Otherwise, we'll end up with a lot of garbage in the target file.
+		if (!Console.IsOutputRedirected)
+		{
+			ResetTimer();
+		}
+	}
+
+	public void Report(double value)
+	{
+		// Make sure value is in [0..1] range
+		value = Math.Max(0, Math.Min(1, value));
+		Interlocked.Exchange(ref currentProgress, value);
+	}
+
+	private void TimerHandler(object state)
+	{
+		lock (timer)
+		{
+			if (disposed) return;
+
+			int progressBlockCount = (int)(currentProgress * blockCount);
+			int percent = (int)(currentProgress * 100);
+			string text = string.Format("[{0}{1}] {2,3}% {3}",
+				new string('#', progressBlockCount), new string('-', blockCount - progressBlockCount),
+				percent,
+				animation[animationIndex++ % animation.Length]);
+			UpdateText(text);
+
+			ResetTimer();
+		}
+	}
+
+	private void UpdateText(string text)
+	{
+		// Get length of common portion
+		int commonPrefixLength = 0;
+		int commonLength = Math.Min(currentText.Length, text.Length);
+		while (commonPrefixLength < commonLength && text[commonPrefixLength] == currentText[commonPrefixLength])
+		{
+			commonPrefixLength++;
+		}
+
+		// Backtrack to the first differing character
+		StringBuilder outputBuilder = new StringBuilder();
+		outputBuilder.Append('\b', currentText.Length - commonPrefixLength);
+
+		// Output new suffix
+		outputBuilder.Append(text.Substring(commonPrefixLength));
+
+		// If the new text is shorter than the old one: delete overlapping characters
+		int overlapCount = currentText.Length - text.Length;
+		if (overlapCount > 0)
+		{
+			outputBuilder.Append(' ', overlapCount);
+			outputBuilder.Append('\b', overlapCount);
+		}
+
+		Console.Write(outputBuilder);
+		currentText = text;
+	}
+
+	private void ResetTimer()
+	{
+		timer.Change(animationInterval, TimeSpan.FromMilliseconds(-1));
+	}
+
+	public void Dispose()
+	{
+		lock (timer)
+		{
+			disposed = true;
+			UpdateText(string.Empty);
+		}
+	}
+
+}
\ No newline at end of file
diff --git a/Engine/Programs/SkeinCLI/SkeinCLI/ProjectsCmd.cs b/Engine/Programs/SkeinCLI/SkeinCLI/ProjectsCmd.cs
new file mode 100644
--- /dev/null
+++ b/Engine/Programs/SkeinCLI/SkeinCLI/ProjectsCmd.cs
@@ -0,0 +1,242 @@
+﻿using McMaster.Extensions.CommandLineUtils;
+using Serilog;
+using System;
+using System.Collections.Generic;
+using System.ComponentModel.DataAnnotations;
+using System.IO;
+using System.Threading.Tasks;
+
+namespace SkeinCLI
+{
+
+    [Command(Name = "projects", Description = "Executes project related commands")]
+    [Subcommand(
+        typeof(ProjectsListCmd),
+        typeof(ProjectsSnapshotsCmd),
+        typeof(ProjectsStatusCmd),
+        typeof(ProjectsRevertCmd))]
+    class ProjectsCmd : SkeinCmdBase
+    {
+        // This will automatically be set before OnExecute is invoked.
+        private SkeinCmd Parent { get; set; }
+
+        public override List<string> CreateArgs()
+        {
+            var args = Parent.CreateArgs();
+            args.Add("projects");
+
+            return args;
+        }
+
+        protected override Task<int> OnExecute(CommandLineApplication app)
+        {
+            //'skein projects' is an incomplete command, so we show help here even if --help is not specified
+            app.ShowHelp();
+            return base.OnExecute(app);
+        }
+    }
+
+    #region projects list
+
+    [Command(Name = "list", Description = "Get a list of projects")]
+    class ProjectsListCmd : SkeinCmdBase
+    {
+        private ProjectsCmd Parent { get; set; }
+
+        public override List<string> CreateArgs()
+        {
+            var args = Parent.CreateArgs();
+            args.Add("list");
+
+            return args;
+        }
+
+        protected override Task<int> OnExecute(CommandLineApplication app)
+        {
+            Log.Logger.ForContext<ProjectsListCmd>().Information("Projects listed here [...]");
+            return base.OnExecute(app);
+        }
+    }
+
+    #endregion
+
+    #region projects snapshots
+
+    [Command(Name = "snapshots", Description = "Executes project snapshots related commands")]
+    [Subcommand(
+        typeof(ProjectsSnapshotsGetCmd),
+        typeof(ProjectsSnapshotsCreateCmd))]
+    class ProjectsSnapshotsCmd : SkeinCmdBase
+    {
+        private ProjectsCmd Parent { get; set; }
+
+        public override List<string> CreateArgs()
+        {
+            var args = Parent.CreateArgs();
+            args.Add("snapshots");
+
+            return args;
+        }
+
+        protected override Task<int> OnExecute(CommandLineApplication app)
+        {
+            //'skein projects snapshots' is an incomplete command, so we show help here even if --help is not specified
+            app.ShowHelp();
+            return base.OnExecute(app);
+        }
+    }
+
+    [Command(Name = "get", Description = "Download a project snapshot locally")]
+    class ProjectsSnapshotsGetCmd : SkeinCmdBase
+    {
+        private ProjectsSnapshotsCmd Parent { get; set; }
+
+        [Argument(0, Name = "project-id", Description = "uuid of the project to get")]
+        [Required]
+        public string ProjectId { get; set; }
+
+        [Argument(1, Name = "path", Description = "Local path where to save the snapshot, defaults to the current path")]
+        [DirectoryExists]
+        public string Path { get; set; }
+
+        [Option("-s|--snapshot", CommandOptionType.SingleValue, Description = "Version of the snapshot to download")]
+        public string Snapshot { get; set; }
+
+        public override List<string> CreateArgs()
+        {
+            var args = Parent.CreateArgs();
+            args.Add("get");
+
+            return args;
+        }
+
+        protected override Task<int> OnExecute(CommandLineApplication app)
+        {
+            if (string.IsNullOrEmpty(Path))
+            {
+                Path = Directory.GetCurrentDirectory();
+            }
+            Log.Logger.ForContext<ProjectsSnapshotsGetCmd>().Information($"Snapshot: {Snapshot}, ProjectId: {ProjectId}, Path: {Path}");
+            Log.Logger.ForContext<ProjectsSnapshotsGetCmd>().Information("Download progress shown here [...]");
+            return base.OnExecute(app);
+        }
+    }
+
+    [Command(Name = "create", Description = "Create a snapshot of a project and upload assets")]
+    class ProjectsSnapshotsCreateCmd : SkeinCmdBase
+    {
+        private ProjectsSnapshotsCmd Parent { get; set; }
+
+        [Argument(0, Name = "path", Description = "Local path from where to create the snapshot, defaults to the current path")]
+        [DirectoryExists]
+        public string Path { get; set; }
+
+        [Option("-m|--message", CommandOptionType.SingleValue, Description = "snapshot description")]
+        [Required]
+        public string Message { get; set; }
+
+        public override List<string> CreateArgs()
+        {
+            var args = Parent.CreateArgs();
+            args.Add("create");
+
+            return args;
+        }
+
+        protected override Task<int> OnExecute(CommandLineApplication app)
+        {
+            if (string.IsNullOrEmpty(Path))
+            {
+                Path = Directory.GetCurrentDirectory();
+            }
+            Log.Logger.ForContext<ProjectsSnapshotsGetCmd>().Information($"Path: {Path}, Message: {Message}");
+            Log.Logger.ForContext<ProjectsSnapshotsCreateCmd>().Information("Create operation result shown here [...]");
+            return base.OnExecute(app);
+        }
+    }
+
+    #endregion
+    
+    #region projects status
+
+    [Command(Name = "status", Description = "List the all project assets and if they are modified locally")]
+    class ProjectsStatusCmd : SkeinCmdBase
+    {
+        private ProjectsCmd Parent { get; set; }
+
+        [Argument(0, Name = "path", Description = "Local path for which to request the status, defaults to the current path")]
+        [DirectoryExists]
+        public string Path { get; set; }
+
+        public override List<string> CreateArgs()
+        {
+            var args = Parent.CreateArgs();
+            args.Add("status");
+
+            return args;
+        }
+
+        protected override Task<int> OnExecute(CommandLineApplication app)
+        {
+            if (string.IsNullOrEmpty(Path))
+            {
+                Path = Directory.GetCurrentDirectory();
+            }
+            Log.Logger.ForContext<ProjectsSnapshotsGetCmd>().Information($"Path: {Path}");
+            Log.Logger.ForContext<ProjectsStatusCmd>().Information("Status described here [...]");
+            return base.OnExecute(app);
+        }
+    }
+
+    #endregion
+
+    #region projects revert
+
+    [Command(Name = "revert", Description = "Revert the project to the snapshot")]
+    class ProjectsRevertCmd : SkeinCmdBase
+    {
+        private ProjectsCmd Parent { get; set; }
+
+        [Argument(0, Name = "path", Description = "Local path for which to revert the project, defaults to the current path")]
+        [DirectoryExists]
+        public string Path { get; set; }
+
+        public override List<string> CreateArgs()
+        {
+            var args = Parent.CreateArgs();
+            args.Add("revert");
+
+            return args;
+        }
+
+        protected override Task<int> OnExecute(CommandLineApplication app)
+        {
+            if (string.IsNullOrEmpty(Path))
+            {
+                Path = Directory.GetCurrentDirectory();
+            }
+
+            bool proceed = true;
+            if (!Quiet)
+            {
+                proceed = Prompt.GetYesNo($"Are you sure you want to revert the project at {Path}?",
+                    defaultAnswer: true,
+                    promptColor: ConsoleColor.Black,
+                    promptBgColor: ConsoleColor.White);
+            }
+
+            if (proceed)
+            {
+                Log.Logger.ForContext<ProjectsSnapshotsGetCmd>().Information($"Path: {Path}");
+                Log.Logger.ForContext<ProjectsRevertCmd>().Information("Revert operation result here [...]");
+            }
+            else 
+            {
+                Log.Logger.ForContext<ProjectsRevertCmd>().Information("User cancelled operation");
+            }
+            return base.OnExecute(app);
+        }
+    }
+
+    #endregion
+}
\ No newline at end of file
diff --git a/Engine/Programs/SkeinCLI/SkeinCLI/SkeinCLI.csproj b/Engine/Programs/SkeinCLI/SkeinCLI/SkeinCLI.csproj
new file mode 100644
--- /dev/null
+++ b/Engine/Programs/SkeinCLI/SkeinCLI/SkeinCLI.csproj
@@ -0,0 +1,33 @@
+﻿<Project Sdk="Microsoft.NET.Sdk">
+
+  <PropertyGroup>
+    <OutputType>Exe</OutputType>
+    <TargetFramework>netcoreapp3.1</TargetFramework>
+    <RuntimeIdentifier>win7-x64</RuntimeIdentifier>
+    <LangVersion>latest</LangVersion>
+    <AssemblyName>skein</AssemblyName>
+  </PropertyGroup>
+
+  <ItemGroup>
+    <None Remove="appsettings.json" />
+  </ItemGroup>
+
+  <ItemGroup>
+    <Content Include="appsettings.json">
+      <CopyToOutputDirectory>Always</CopyToOutputDirectory>
+    </Content>
+  </ItemGroup>
+
+  <ItemGroup>
+    <PackageReference Include="McMaster.Extensions.CommandLineUtils" Version="3.1.0" />
+    <PackageReference Include="Microsoft.Extensions.Configuration" Version="5.0.0" />
+    <PackageReference Include="Microsoft.Extensions.Configuration.EnvironmentVariables" Version="5.0.0" />
+    <PackageReference Include="Microsoft.Extensions.Configuration.Json" Version="5.0.0" />
+    <PackageReference Include="Microsoft.Extensions.Http" Version="5.0.0" />
+    <PackageReference Include="Serilog" Version="2.10.0" />
+    <PackageReference Include="Serilog.Settings.Configuration" Version="3.2.0" />
+    <PackageReference Include="Serilog.Sinks.Console" Version="4.0.0" />
+    <PackageReference Include="Serilog.Sinks.File" Version="5.0.0" />
+  </ItemGroup>
+
+</Project>
diff --git a/Engine/Programs/SkeinCLI/SkeinCLI/SkeinClient.cs b/Engine/Programs/SkeinCLI/SkeinCLI/SkeinClient.cs
new file mode 100644
--- /dev/null
+++ b/Engine/Programs/SkeinCLI/SkeinCLI/SkeinClient.cs
@@ -0,0 +1,26 @@
+﻿using System.Net.Http;
+using System.Threading.Tasks;
+
+public class SkeinClient
+{
+    private readonly HttpClient _httpClient;
+
+    public SkeinClient(HttpClient httpClient)
+    {
+        _httpClient = httpClient;
+    }
+
+    public async Task<string> GetAsync(string url)
+    {
+        HttpRequestMessage requestMessage = new HttpRequestMessage(HttpMethod.Get, url);
+
+        return await Request(requestMessage);
+    }
+
+    private async Task<string> Request(HttpRequestMessage requestMessage)
+    {
+        // Mock Response
+        await Task.Delay(1000);
+        return "{}";
+    }
+}
diff --git a/Engine/Programs/SkeinCLI/SkeinCLI/SkeinCmd.cs b/Engine/Programs/SkeinCLI/SkeinCLI/SkeinCmd.cs
new file mode 100644
--- /dev/null
+++ b/Engine/Programs/SkeinCLI/SkeinCLI/SkeinCmd.cs
@@ -0,0 +1,32 @@
+﻿using McMaster.Extensions.CommandLineUtils;
+using System.Collections.Generic;
+using System.Reflection;
+using System.Threading.Tasks;
+
+namespace SkeinCLI
+{
+    [Command(Name = "skein", UnrecognizedArgumentHandling = UnrecognizedArgumentHandling.Throw, OptionsComparison = System.StringComparison.InvariantCultureIgnoreCase)]
+    [VersionOptionFromMember("--version", MemberName = nameof(GetVersion))]
+    [Subcommand(
+        typeof(AuthCmd),
+        typeof(ProjectsCmd),
+        typeof(AssetsCmd))]
+    class SkeinCmd : SkeinCmdBase
+    {
+        protected override Task<int> OnExecute(CommandLineApplication app)
+        {
+            // this shows help even if the --help option isn't specified
+            app.ShowHelp();
+            return Task.FromResult(0);
+        }
+
+        private static string GetVersion()
+            => typeof(SkeinCmd).Assembly.GetCustomAttribute<AssemblyInformationalVersionAttribute>().InformationalVersion;
+
+        public override List<string> CreateArgs()
+        {
+            var args = new List<string>();
+            return args;
+        }
+    }
+}
\ No newline at end of file
diff --git a/Engine/Programs/SkeinCLI/SkeinCLI/SkeinCmdBase.cs b/Engine/Programs/SkeinCLI/SkeinCLI/SkeinCmdBase.cs
new file mode 100644
--- /dev/null
+++ b/Engine/Programs/SkeinCLI/SkeinCLI/SkeinCmdBase.cs
@@ -0,0 +1,34 @@
+﻿using McMaster.Extensions.CommandLineUtils;
+using Serilog;
+using System;
+using System.Collections.Generic;
+using System.Threading.Tasks;
+
+namespace SkeinCLI
+{
+    [HelpOption("--help")]
+    abstract class SkeinCmdBase
+    {
+        public abstract List<string> CreateArgs();
+        
+        [Option("-q|--quiet", CommandOptionType.NoValue, Description = "Execute in quiet mode")]
+        public bool Quiet { get; set; }
+
+        protected virtual Task<int> OnExecute(CommandLineApplication app)
+        {
+            var args = CreateArgs();
+            if (Quiet)
+            {
+                args.Add("-q");
+            }
+
+            Log.Logger.ForContext<SkeinCmdBase>().Debug("Executing: skein " + ArgumentEscaper.EscapeAndConcatenate(args));
+            return Task.FromResult(0);
+        }
+
+        protected void OnException(Exception ex)
+        {
+            Log.Logger.ForContext<SkeinCmdBase>().Error(ex, "Exception");
+        }
+    }
+}
\ No newline at end of file
diff --git a/Engine/Programs/SkeinCLI/SkeinCLI/appsettings.json b/Engine/Programs/SkeinCLI/SkeinCLI/appsettings.json
new file mode 100644
--- /dev/null
+++ b/Engine/Programs/SkeinCLI/SkeinCLI/appsettings.json
@@ -0,0 +1,22 @@
+{
+  "Serilog": {
+    "MinimumLevel": "Debug",
+    "Enrich": [ "FromLogContext", "WithExceptionDetails" ],
+    "WriteTo": [
+      {
+        "Name": "Console",
+        "Args": {
+          "outputTemplate": "{Timestamp:HH:mm:ss.fff zzz} [{Level:w3}] {Message:lj}{NewLine}{Exception}"
+        }
+      },
+      {
+        "Name": "File",
+        "Args": {
+          "path": "Logs/skein.log",
+          "rollingInterval": "Day",
+          "outputTemplate": "{Timestamp:HH:mm:ss.fff zzz} [{Level:w3}] {Message:lj}{NewLine}{Exception}"
+        }
+      }
+    ]
+  }
+}
\ No newline at end of file
